Added SimpleHMD plugin to experimental.

This is basically a bare minimum HMD sample that can be used as a jumping off point for implementing a new HMD integration.  When enabled it will run stereo with a fixed IPD, and will also run a test mesh distort if post is enabled.

diff --git a/Engine/Plugins/Experimental/SimpleHMD/SimpleHMD.uplugin b/Engine/Plugins/Experimental/SimpleHMD/SimpleHMD.uplugin
new file mode 100644
--- /dev/null
+++ b/Engine/Plugins/Experimental/SimpleHMD/SimpleHMD.uplugin
@@ -0,0 +1,22 @@
+{
+	"FileVersion" : 3,
+
+	"FriendlyName" : "SimpleHMD Plugin",
+	"Version" : 1,
+	"VersionName" : "1.0",
+	"CreatedBy" : "Epic Games, Inc.",
+	"CreatedByURL" : "http://epicgames.com",
+	"EngineVersion" : "4.2.0",
+	"Description" : "SimpleHMD is a sample of a basic stereo HMD implementation",
+	"Category" : "Head Mounted Displays.Devices",
+	"EnabledByDefault" : false,
+
+	"Modules" :
+	[
+		{
+			"Name" : "SimpleHMD",
+			"Type" : "Runtime",
+			"WhitelistPlatforms" : [ "Win64", "Win32", "Mac", "Android" ]
+		}
+	]
+}
diff --git a/Engine/Plugins/Experimental/SimpleHMD/Source/SimpleHMD/Private/SimpleHMD.cpp b/Engine/Plugins/Experimental/SimpleHMD/Source/SimpleHMD/Private/SimpleHMD.cpp
new file mode 100644
--- /dev/null
+++ b/Engine/Plugins/Experimental/SimpleHMD/Source/SimpleHMD/Private/SimpleHMD.cpp
@@ -0,0 +1,353 @@
+// Copyright 1998-2014 Epic Games, Inc. All Rights Reserved.
+
+#include "SimpleHMDPrivatePCH.h"
+#include "SimpleHMD.h"
+#include "RendererPrivate.h"
+#include "ScenePrivate.h"
+#include "PostProcess/PostProcessHMD.h"
+
+//---------------------------------------------------
+// SimpleHMD Plugin Implementation
+//---------------------------------------------------
+
+class FSimpleHMDPlugin : public ISimpleHMDPlugin
+{
+	/** IHeadMountedDisplayModule implementation */
+	virtual TSharedPtr< class IHeadMountedDisplay > CreateHeadMountedDisplay() override;
+
+	FString GetModulePriorityKeyName() const
+	{
+		return FString(TEXT("SimpleHMD"));
+	}
+};
+
+IMPLEMENT_MODULE( FSimpleHMDPlugin, SimpleHMD )
+
+TSharedPtr< class IHeadMountedDisplay > FSimpleHMDPlugin::CreateHeadMountedDisplay()
+{
+#if SIMPLEHMD_SUPPORTED_PLATFORMS
+	TSharedPtr< FSimpleHMD > SimpleHMD( new FSimpleHMD() );
+	if( SimpleHMD->IsInitialized() )
+	{
+		return SimpleHMD;
+	}
+#endif//SIMPLEHMD_SUPPORTED_PLATFORMS
+	return NULL;
+}
+
+
+//---------------------------------------------------
+// SimpleHMD IHeadMountedDisplay Implementation
+//---------------------------------------------------
+
+#if SIMPLEHMD_SUPPORTED_PLATFORMS
+
+bool FSimpleHMD::IsHMDEnabled() const
+{
+	return true;
+}
+
+void FSimpleHMD::EnableHMD(bool enable)
+{
+}
+
+EHMDDeviceType::Type FSimpleHMD::GetHMDDeviceType() const
+{
+	return EHMDDeviceType::DT_ES2GenericStereoMesh;
+}
+
+bool FSimpleHMD::GetHMDMonitorInfo(MonitorInfo& MonitorDesc)
+{
+	MonitorDesc.MonitorName = "";
+	MonitorDesc.MonitorId = 0;
+	MonitorDesc.DesktopX = MonitorDesc.DesktopY = MonitorDesc.ResolutionX = MonitorDesc.ResolutionY = 0;
+	return false;
+}
+
+void FSimpleHMD::GetFieldOfView(float& OutHFOVInDegrees, float& OutVFOVInDegrees) const
+{
+	OutHFOVInDegrees = 0.0f;
+	OutVFOVInDegrees = 0.0f;
+}
+
+bool FSimpleHMD::DoesSupportPositionalTracking() const
+{
+	return false;
+}
+
+bool FSimpleHMD::HasValidTrackingPosition() const
+{
+	return false;
+}
+
+void FSimpleHMD::GetPositionalTrackingCameraProperties(FVector& OutOrigin, FRotator& OutOrientation, float& OutHFOV, float& OutVFOV, float& OutCameraDistance, float& OutNearPlane, float& OutFarPlane) const
+{
+}
+
+void FSimpleHMD::SetInterpupillaryDistance(float NewInterpupillaryDistance)
+{
+}
+
+float FSimpleHMD::GetInterpupillaryDistance() const
+{
+	return 0.064f;
+}
+
+void FSimpleHMD::GetCurrentOrientationAndPosition(FQuat& CurrentOrientation, FVector& CurrentPosition)
+{
+	CurrentOrientation = FQuat(FRotator(0.0f,0.0f,0.0f));
+	CurrentPosition = FVector(0.0f,0.0f,0.0f);
+}
+
+ISceneViewExtension* FSimpleHMD::GetViewExtension()
+{
+	return this;
+}
+
+void FSimpleHMD::ApplyHmdRotation(APlayerController* PC, FRotator& ViewRotation)
+{
+	return;
+}
+
+void FSimpleHMD::UpdatePlayerCameraRotation(APlayerCameraManager* Camera, struct FMinimalViewInfo& POV)
+{
+	return;
+}
+
+bool FSimpleHMD::IsChromaAbCorrectionEnabled() const
+{
+	return false;
+}
+
+bool FSimpleHMD::Exec( UWorld* InWorld, const TCHAR* Cmd, FOutputDevice& Ar )
+{
+	return false;
+}
+
+void FSimpleHMD::OnScreenModeChange(EWindowMode::Type WindowMode)
+{
+}
+
+bool FSimpleHMD::IsPositionalTrackingEnabled() const
+{
+	return false;
+}
+
+bool FSimpleHMD::EnablePositionalTracking(bool enable)
+{
+	return false;
+}
+
+bool FSimpleHMD::IsHeadTrackingAllowed() const
+{
+	return false;
+}
+
+bool FSimpleHMD::IsInLowPersistenceMode() const
+{
+	return false;
+}
+
+void FSimpleHMD::EnableLowPersistenceMode(bool Enable)
+{
+}
+
+void FSimpleHMD::ResetOrientationAndPosition(float yaw)
+{
+	ResetOrientation(yaw);
+	ResetPosition();
+}
+
+void FSimpleHMD::ResetOrientation(float Yaw)
+{
+}
+void FSimpleHMD::ResetPosition()
+{
+}
+
+void FSimpleHMD::SetClippingPlanes(float NCP, float FCP)
+{
+}
+
+void FSimpleHMD::SetBaseRotation(const FRotator& BaseRot)
+{
+}
+
+FRotator FSimpleHMD::GetBaseRotation() const
+{
+	return FRotator::ZeroRotator;
+}
+
+void FSimpleHMD::SetBaseOrientation(const FQuat& BaseOrient)
+{
+}
+
+FQuat FSimpleHMD::GetBaseOrientation() const
+{
+	return FQuat::Identity;
+}
+
+void FSimpleHMD::SetPositionOffset(const FVector& PosOff)
+{
+}
+
+FVector FSimpleHMD::GetPositionOffset() const
+{
+	return FVector::ZeroVector;
+}
+
+void FSimpleHMD::DrawDistortionMesh_RenderThread(struct FRenderingCompositePassContext& Context, const FSceneView& View, const FIntPoint& TextureSize)
+{
+	float ClipSpaceQuadZ = 0.0f;
+	FMatrix QuadTexTransform = FMatrix::Identity;
+	FMatrix QuadPosTransform = FMatrix::Identity;
+	const FIntRect SrcRect = View.ViewRect;
+
+	FRHICommandListImmediate& RHICmdList = Context.RHICmdList;
+	const FSceneViewFamily& ViewFamily = *(View.Family);
+	FIntPoint ViewportSize = ViewFamily.RenderTarget->GetSizeXY();
+	RHICmdList.SetViewport(0, 0, 0.0f, ViewportSize.X, ViewportSize.Y, 1.0f);
+
+	static const uint32 NumVerts = 8;
+	static const uint32 NumTris = 4;
+
+	static const FDistortionVertex Verts[8] =
+	{
+		// left eye
+		{ FVector2D(-0.9f, -0.9f), FVector2D(0.0f, 1.0f), FVector2D(0.0f, 1.0f), FVector2D(0.0f, 1.0f), 1.0f, 0.0f },
+		{ FVector2D(-0.1f, -0.9f), FVector2D(0.5f, 1.0f), FVector2D(0.5f, 1.0f), FVector2D(0.5f, 1.0f), 1.0f, 0.0f },
+		{ FVector2D(-0.1f, 0.9f), FVector2D(0.5f, 0.0f), FVector2D(0.5f, 0.0f), FVector2D(0.5f, 0.0f), 1.0f, 0.0f },
+		{ FVector2D(-0.9f, 0.9f), FVector2D(0.0f, 0.0f), FVector2D(0.0f, 0.0f), FVector2D(0.0f, 0.0f), 1.0f, 0.0f },
+		// right eye
+		{ FVector2D(0.1f, -0.9f), FVector2D(0.5f, 1.0f), FVector2D(0.5f, 1.0f), FVector2D(0.5f, 1.0f), 1.0f, 0.0f },
+		{ FVector2D(0.9f, -0.9f), FVector2D(1.0f, 1.0f), FVector2D(1.0f, 1.0f), FVector2D(1.0f, 1.0f), 1.0f, 0.0f },
+		{ FVector2D(0.9f, 0.9f), FVector2D(1.0f, 0.0f), FVector2D(1.0f, 0.0f), FVector2D(1.0f, 0.0f), 1.0f, 0.0f },
+		{ FVector2D(0.1f, 0.9f), FVector2D(0.5f, 0.0f), FVector2D(0.5f, 0.0f), FVector2D(0.5f, 0.0f), 1.0f, 0.0f },
+	};
+
+	static const uint16 Indices[12] = { /*Left*/ 0, 1, 2, 0, 2, 3, /*Right*/ 4, 5, 6, 4, 6, 7 };
+
+	DrawIndexedPrimitiveUP(Context.RHICmdList, PT_TriangleList, 0, NumVerts, NumTris, &Indices,
+		sizeof(Indices[0]), &Verts, sizeof(Verts[0]));
+}
+	
+void FSimpleHMD::UpdateScreenSettings(const FViewport*)
+{
+}
+
+bool FSimpleHMD::IsStereoEnabled() const
+{
+	return true;
+}
+
+bool FSimpleHMD::EnableStereo(bool stereo)
+{
+	return true;
+}
+
+void FSimpleHMD::AdjustViewRect(EStereoscopicPass StereoPass, int32& X, int32& Y, uint32& SizeX, uint32& SizeY) const
+{
+	SizeX = SizeX / 2;
+	if( StereoPass == eSSP_RIGHT_EYE )
+	{
+		X += SizeX;
+	}
+}
+
+void FSimpleHMD::CalculateStereoViewOffset(const enum EStereoscopicPass StereoPassType, const FRotator& ViewRotation, const float WorldToMeters, FVector& ViewLocation)
+{
+	if( StereoPassType != eSSP_FULL)
+	{
+		float EyeOffset = 3.20000005f;
+		const float PassOffset = (StereoPassType == eSSP_LEFT_EYE) ? EyeOffset : -EyeOffset;
+		ViewLocation += ViewRotation.Quaternion().RotateVector(FVector(0,PassOffset,0));
+	}
+}
+
+FMatrix FSimpleHMD::GetStereoProjectionMatrix(const enum EStereoscopicPass StereoPassType, const float FOV) const
+{
+	const float ProjectionCenterOffset = 0.151976421f;
+	const float PassProjectionOffset = (StereoPassType == eSSP_LEFT_EYE) ? ProjectionCenterOffset : -ProjectionCenterOffset;
+
+	const float HalfFov = 2.19686294f / 2.f;
+	const float InWidth = 640.f;
+	const float InHeight = 480.f;
+	const float XS = 1.0f / tan(HalfFov);
+	const float YS = InWidth / tan(HalfFov) / InHeight;
+
+	const float InNearZ = GNearClippingPlane;
+	return FMatrix(
+		FPlane(XS,                      0.0f,								    0.0f,							0.0f),
+		FPlane(0.0f,					YS,	                                    0.0f,							0.0f),
+		FPlane(0.0f,	                0.0f,								    0.0f,							1.0f),
+		FPlane(0.0f,					0.0f,								    InNearZ,						0.0f))
+
+		* FTranslationMatrix(FVector(PassProjectionOffset,0,0));
+}
+
+void FSimpleHMD::InitCanvasFromView(FSceneView* InView, UCanvas* Canvas)
+{
+}
+
+void FSimpleHMD::PushViewportCanvas(EStereoscopicPass StereoPass, FCanvas *InCanvas, UCanvas *InCanvasObject, FViewport *InViewport) const 
+{
+	FMatrix m;
+	m.SetIdentity();
+	InCanvas->PushAbsoluteTransform(m);
+}
+
+void FSimpleHMD::PushViewCanvas(EStereoscopicPass StereoPass, FCanvas *InCanvas, UCanvas *InCanvasObject, FSceneView *InView) const 
+{
+	FMatrix m;
+	m.SetIdentity();
+	InCanvas->PushAbsoluteTransform(m);
+}
+
+void FSimpleHMD::GetEyeRenderParams_RenderThread(EStereoscopicPass StereoPass, FVector2D& EyeToSrcUVScaleValue, FVector2D& EyeToSrcUVOffsetValue) const
+{
+	EyeToSrcUVOffsetValue = FVector2D::ZeroVector;
+	EyeToSrcUVScaleValue = FVector2D(1.0f, 1.0f);
+}
+
+
+void FSimpleHMD::ModifyShowFlags(FEngineShowFlags& ShowFlags)
+{
+	ShowFlags.MotionBlur = 0;
+	ShowFlags.HMDDistortion = true;
+	ShowFlags.ScreenPercentage = 1.0f;
+	ShowFlags.StereoRendering = IsStereoEnabled();
+}
+
+void FSimpleHMD::SetupView(FSceneViewFamily& InViewFamily, FSceneView& InView)
+{
+	InView.BaseHmdOrientation = FQuat(FRotator(0.0f,0.0f,0.0f));
+	InView.BaseHmdLocation = FVector(0.f);
+//	WorldToMetersScale = InView.WorldToMetersScale;
+	InViewFamily.bUseSeparateRenderTarget = false;
+}
+
+void FSimpleHMD::PreRenderView_RenderThread(FSceneView& View)
+{
+	check(IsInRenderingThread());
+}
+
+void FSimpleHMD::PreRenderViewFamily_RenderThread(FSceneViewFamily& ViewFamily)
+{
+	check(IsInRenderingThread());
+}
+
+FSimpleHMD::FSimpleHMD()
+{
+}
+
+FSimpleHMD::~FSimpleHMD()
+{
+}
+
+bool FSimpleHMD::IsInitialized() const
+{
+	return true;
+}
+
+
+#endif //SIMPLEHMD_SUPPORTED_PLATFORMS
diff --git a/Engine/Plugins/Experimental/SimpleHMD/Source/SimpleHMD/Private/SimpleHMD.h b/Engine/Plugins/Experimental/SimpleHMD/Source/SimpleHMD/Private/SimpleHMD.h
new file mode 100644
--- /dev/null
+++ b/Engine/Plugins/Experimental/SimpleHMD/Source/SimpleHMD/Private/SimpleHMD.h
@@ -0,0 +1,103 @@
+// Copyright 1998-2014 Epic Games, Inc. All Rights Reserved.
+
+#pragma once
+#include "ISimpleHMDPlugin.h"
+#include "HeadMountedDisplay.h"
+#include "IHeadMountedDisplay.h"
+
+#if SIMPLEHMD_SUPPORTED_PLATFORMS
+
+#include "SceneViewExtension.h"
+
+
+/**
+ * Simple Head Mounted Display
+ */
+class FSimpleHMD : public IHeadMountedDisplay, public ISceneViewExtension
+{
+public:
+	/** IHeadMountedDisplay interface */
+	virtual bool IsHMDConnected() override { return true; }
+	virtual bool IsHMDEnabled() const override;
+	virtual void EnableHMD(bool allow = true) override;
+	virtual EHMDDeviceType::Type GetHMDDeviceType() const override;
+	virtual bool GetHMDMonitorInfo(MonitorInfo&) override;
+
+	virtual void GetFieldOfView(float& OutHFOVInDegrees, float& OutVFOVInDegrees) const override;
+
+	virtual bool DoesSupportPositionalTracking() const override;
+	virtual bool HasValidTrackingPosition() const override;
+	virtual void GetPositionalTrackingCameraProperties(FVector& OutOrigin, FRotator& OutOrientation, float& OutHFOV, float& OutVFOV, float& OutCameraDistance, float& OutNearPlane, float& OutFarPlane) const override;
+
+	virtual void SetInterpupillaryDistance(float NewInterpupillaryDistance) override;
+	virtual float GetInterpupillaryDistance() const override;
+
+	virtual void GetCurrentOrientationAndPosition(FQuat& CurrentOrientation, FVector& CurrentPosition) override;
+	virtual class ISceneViewExtension* GetViewExtension() override;
+	virtual void ApplyHmdRotation(APlayerController* PC, FRotator& ViewRotation) override;
+	virtual void UpdatePlayerCameraRotation(APlayerCameraManager*, struct FMinimalViewInfo& POV) override;
+
+	virtual bool IsChromaAbCorrectionEnabled() const override;
+
+	virtual bool Exec( UWorld* InWorld, const TCHAR* Cmd, FOutputDevice& Ar ) override;
+	virtual void OnScreenModeChange(EWindowMode::Type WindowMode) override;
+
+	virtual bool IsPositionalTrackingEnabled() const override;
+	virtual bool EnablePositionalTracking(bool enable) override;
+
+	virtual bool IsHeadTrackingAllowed() const override;
+
+	virtual bool IsInLowPersistenceMode() const override;
+	virtual void EnableLowPersistenceMode(bool Enable = true) override;
+
+	virtual void ResetOrientationAndPosition(float yaw = 0.f) override;
+	virtual void ResetOrientation(float Yaw = 0.f) override;
+	virtual void ResetPosition() override;
+
+	virtual void SetClippingPlanes(float NCP, float FCP) override;
+
+	virtual void SetBaseRotation(const FRotator& BaseRot) override;
+	virtual FRotator GetBaseRotation() const override;
+
+	virtual void SetBaseOrientation(const FQuat& BaseOrient) override;
+	virtual FQuat GetBaseOrientation() const override;
+
+	virtual void SetPositionOffset(const FVector& PosOff) override;
+	virtual FVector GetPositionOffset() const override;
+
+	virtual void DrawDistortionMesh_RenderThread(struct FRenderingCompositePassContext& Context, const FSceneView& View, const FIntPoint& TextureSize) override;
+	virtual void UpdateScreenSettings(const FViewport*) override;
+
+	/** IStereoRendering interface */
+	virtual bool IsStereoEnabled() const override;
+	virtual bool EnableStereo(bool stereo = true) override;
+	virtual void AdjustViewRect(EStereoscopicPass StereoPass, int32& X, int32& Y, uint32& SizeX, uint32& SizeY) const override;
+	virtual void CalculateStereoViewOffset(const EStereoscopicPass StereoPassType, const FRotator& ViewRotation, 
+										   const float MetersToWorld, FVector& ViewLocation) override;
+	virtual FMatrix GetStereoProjectionMatrix(const EStereoscopicPass StereoPassType, const float FOV) const override;
+	virtual void InitCanvasFromView(FSceneView* InView, UCanvas* Canvas) override;
+	virtual void PushViewportCanvas(EStereoscopicPass StereoPass, FCanvas *InCanvas, UCanvas *InCanvasObject, FViewport *InViewport) const override;
+	virtual void PushViewCanvas(EStereoscopicPass StereoPass, FCanvas *InCanvas, UCanvas *InCanvasObject, FSceneView *InView) const override;
+	virtual void GetEyeRenderParams_RenderThread(EStereoscopicPass StereoPass, FVector2D& EyeToSrcUVScaleValue, FVector2D& EyeToSrcUVOffsetValue) const override;
+
+	/** ISceneViewExtension interface */
+	virtual void ModifyShowFlags(FEngineShowFlags& ShowFlags) override;
+	virtual void SetupView(FSceneViewFamily& InViewFamily, FSceneView& InView) override;
+	virtual void PreRenderView_RenderThread(FSceneView& InView) override;
+	virtual void PreRenderViewFamily_RenderThread(FSceneViewFamily& InViewFamily) override;
+
+public:
+	/** Constructor */
+	FSimpleHMD();
+
+	/** Destructor */
+	virtual ~FSimpleHMD();
+
+	/** @return	True if the HMD was initialized OK */
+	bool IsInitialized() const;
+};
+
+
+DEFINE_LOG_CATEGORY_STATIC(LogHMD, Log, All);
+
+#endif //SIMPLEHMD_SUPPORTED_PLATFORMS
diff --git a/Engine/Plugins/Experimental/SimpleHMD/Source/SimpleHMD/Private/SimpleHMDPrivatePCH.h b/Engine/Plugins/Experimental/SimpleHMD/Source/SimpleHMD/Private/SimpleHMDPrivatePCH.h
new file mode 100644
--- /dev/null
+++ b/Engine/Plugins/Experimental/SimpleHMD/Source/SimpleHMD/Private/SimpleHMDPrivatePCH.h
@@ -0,0 +1,6 @@
+// Copyright 1998-2014 Epic Games, Inc. All Rights Reserved.
+
+#pragma once
+
+#include "Engine.h"
+#include "IHeadMountedDisplay.h"
diff --git a/Engine/Plugins/Experimental/SimpleHMD/Source/SimpleHMD/Public/ISimpleHMDPlugin.h b/Engine/Plugins/Experimental/SimpleHMD/Source/SimpleHMD/Public/ISimpleHMDPlugin.h
new file mode 100644
--- /dev/null
+++ b/Engine/Plugins/Experimental/SimpleHMD/Source/SimpleHMD/Public/ISimpleHMDPlugin.h
@@ -0,0 +1,39 @@
+// Copyright 1998-2014 Epic Games, Inc. All Rights Reserved.
+
+#pragma once
+
+#include "ModuleManager.h"
+#include "IHeadMountedDisplayModule.h"
+
+#define SIMPLEHMD_SUPPORTED_PLATFORMS PLATFORM_WINDOWS || PLATFORM_MAC || PLATFORM_ANDROID
+
+/**
+ * The public interface to this module.  In most cases, this interface is only public to sibling modules 
+ * within this plugin.
+ */
+class ISimpleHMDPlugin : public IHeadMountedDisplayModule
+{
+
+public:
+
+	/**
+	 * Singleton-like access to this module's interface.  This is just for convenience!
+	 * Beware of calling this during the shutdown phase, though.  Your module might have been unloaded already.
+	 *
+	 * @return Returns singleton instance, loading the module on demand if needed
+	 */
+	static inline ISimpleHMDPlugin& Get()
+	{
+		return FModuleManager::LoadModuleChecked< ISimpleHMDPlugin >( "SimpleHMD" );
+	}
+
+	/**
+	 * Checks to see if this module is loaded and ready.  It is only valid to call Get() if IsAvailable() returns true.
+	 *
+	 * @return True if the module is loaded and ready to use
+	 */
+	static inline bool IsAvailable()
+	{
+		return FModuleManager::Get().IsModuleLoaded( "SimpleHMD" );
+	}
+};
diff --git a/Engine/Plugins/Experimental/SimpleHMD/Source/SimpleHMD/SimpleHMD.Build.cs b/Engine/Plugins/Experimental/SimpleHMD/Source/SimpleHMD/SimpleHMD.Build.cs
new file mode 100644
--- /dev/null
+++ b/Engine/Plugins/Experimental/SimpleHMD/Source/SimpleHMD/SimpleHMD.Build.cs
@@ -0,0 +1,32 @@
+// Copyright 1998-2014 Epic Games, Inc. All Rights Reserved.
+
+namespace UnrealBuildTool.Rules
+{
+	public class SimpleHMD : ModuleRules
+	{
+		public SimpleHMD(TargetInfo Target)
+		{
+			PrivateIncludePaths.AddRange(
+				new string[] {
+					"SimpleHMD/Private",
+					"../../../../Source/Runtime/Renderer/Private",
+					// ... add other private include paths required here ...
+				}
+				);
+
+			PrivateDependencyModuleNames.AddRange(
+				new string[]
+				{
+					"Core",
+					"CoreUObject",
+					"Engine",
+					"RHI",
+					"RenderCore",
+					"Renderer",
+					"ShaderCore",
+					"HeadMountedDisplay"
+				}
+				);
+		}
+	}
+}
