Updated SimpleHMD to build for all platforms as well as enabling extremely basic sensor support.

diff --git a/Engine/Plugins/Experimental/SimpleHMD/SimpleHMD.uplugin b/Engine/Plugins/Experimental/SimpleHMD/SimpleHMD.uplugin
--- a/Engine/Plugins/Experimental/SimpleHMD/SimpleHMD.uplugin
+++ b/Engine/Plugins/Experimental/SimpleHMD/SimpleHMD.uplugin
@@ -15,8 +15,7 @@
 	[
 		{
 			"Name" : "SimpleHMD",
-			"Type" : "Runtime",
-			"WhitelistPlatforms" : [ "Win64", "Win32", "Mac", "Android" ]
+			"Type" : "Runtime"
 		}
 	]
 }
diff --git a/Engine/Plugins/Experimental/SimpleHMD/Source/SimpleHMD/Private/SimpleHMD.cpp b/Engine/Plugins/Experimental/SimpleHMD/Source/SimpleHMD/Private/SimpleHMD.cpp
--- a/Engine/Plugins/Experimental/SimpleHMD/Source/SimpleHMD/Private/SimpleHMD.cpp
+++ b/Engine/Plugins/Experimental/SimpleHMD/Source/SimpleHMD/Private/SimpleHMD.cpp
@@ -25,13 +25,11 @@ IMPLEMENT_MODULE( FSimpleHMDPlugin, SimpleHMD )
 
 TSharedPtr< class IHeadMountedDisplay > FSimpleHMDPlugin::CreateHeadMountedDisplay()
 {
-#if SIMPLEHMD_SUPPORTED_PLATFORMS
 	TSharedPtr< FSimpleHMD > SimpleHMD( new FSimpleHMD() );
 	if( SimpleHMD->IsInitialized() )
 	{
 		return SimpleHMD;
 	}
-#endif//SIMPLEHMD_SUPPORTED_PLATFORMS
 	return NULL;
 }
 
@@ -40,8 +38,6 @@ TSharedPtr< class IHeadMountedDisplay > FSimpleHMDPlugin::CreateHeadMountedDispl
 // SimpleHMD IHeadMountedDisplay Implementation
 //---------------------------------------------------
 
-#if SIMPLEHMD_SUPPORTED_PLATFORMS
-
 bool FSimpleHMD::IsHMDEnabled() const
 {
 	return true;
@@ -93,10 +89,41 @@ float FSimpleHMD::GetInterpupillaryDistance() const
 	return 0.064f;
 }
 
+void FSimpleHMD::GetCurrentPose(FQuat& CurrentOrientation)
+{
+	// very basic.  no head model, no prediction, using debuglocalplayer
+	ULocalPlayer* Player = GEngine->GetDebugLocalPlayer();
+
+	if (Player != NULL && Player->PlayerController != NULL)
+	{
+		FVector RotationRate = Player->PlayerController->GetInputVectorKeyState(EKeys::RotationRate);
+
+		double CurrentTime = FApp::GetCurrentTime();
+		double DeltaTime = 0.0;
+
+		if (LastSensorTime >= 0.0)
+		{
+			DeltaTime = CurrentTime - LastSensorTime;
+		}
+
+		LastSensorTime = CurrentTime;
+
+		// mostly incorrect, but we just want some sensor input for testing
+		RotationRate *= DeltaTime;
+		CurrentOrientation *= FQuat(FRotator(FMath::RadiansToDegrees(-RotationRate.X), FMath::RadiansToDegrees(-RotationRate.Y), FMath::RadiansToDegrees(-RotationRate.Z)));
+	}
+	else
+	{
+		CurrentOrientation = FQuat(FRotator(0.0f, 0.0f, 0.0f));
+	}
+}
+
 void FSimpleHMD::GetCurrentOrientationAndPosition(FQuat& CurrentOrientation, FVector& CurrentPosition)
 {
-	CurrentOrientation = FQuat(FRotator(0.0f,0.0f,0.0f));
-	CurrentPosition = FVector(0.0f,0.0f,0.0f);
+	CurrentPosition = FVector(0.0f, 0.0f, 0.0f);
+
+	GetCurrentPose(CurrentOrientation);
+	CurHmdOrientation = LastHmdOrientation = CurrentOrientation;
 }
 
 ISceneViewExtension* FSimpleHMD::GetViewExtension()
@@ -106,7 +133,21 @@ ISceneViewExtension* FSimpleHMD::GetViewExtension()
 
 void FSimpleHMD::ApplyHmdRotation(APlayerController* PC, FRotator& ViewRotation)
 {
-	return;
+	ViewRotation.Normalize();
+
+	GetCurrentPose(CurHmdOrientation);
+	LastHmdOrientation = CurHmdOrientation;
+
+	const FRotator DeltaRot = ViewRotation - PC->GetControlRotation();
+	DeltaControlRotation = (DeltaControlRotation + DeltaRot).GetNormalized();
+
+	// Pitch from other sources is never good, because there is an absolute up and down that must be respected to avoid motion sickness.
+	// Same with roll.
+	DeltaControlRotation.Pitch = 0;
+	DeltaControlRotation.Roll = 0;
+	DeltaControlOrientation = DeltaControlRotation.Quaternion();
+
+	ViewRotation = FRotator(DeltaControlOrientation * CurHmdOrientation);
 }
 
 void FSimpleHMD::UpdatePlayerCameraRotation(APlayerCameraManager* Camera, struct FMinimalViewInfo& POV)
@@ -140,7 +181,7 @@ bool FSimpleHMD::EnablePositionalTracking(bool enable)
 
 bool FSimpleHMD::IsHeadTrackingAllowed() const
 {
-	return false;
+	return true;
 }
 
 bool FSimpleHMD::IsInLowPersistenceMode() const
@@ -336,7 +377,12 @@ void FSimpleHMD::PreRenderViewFamily_RenderThread(FSceneViewFamily& ViewFamily)
 	check(IsInRenderingThread());
 }
 
-FSimpleHMD::FSimpleHMD()
+FSimpleHMD::FSimpleHMD() :
+	CurHmdOrientation(FQuat::Identity),
+	LastHmdOrientation(FQuat::Identity),
+	DeltaControlRotation(FRotator::ZeroRotator),
+	DeltaControlOrientation(FQuat::Identity),
+	LastSensorTime(-1.0)
 {
 }
 
@@ -348,6 +394,3 @@ bool FSimpleHMD::IsInitialized() const
 {
 	return true;
 }
-
-
-#endif //SIMPLEHMD_SUPPORTED_PLATFORMS
diff --git a/Engine/Plugins/Experimental/SimpleHMD/Source/SimpleHMD/Private/SimpleHMD.h b/Engine/Plugins/Experimental/SimpleHMD/Source/SimpleHMD/Private/SimpleHMD.h
--- a/Engine/Plugins/Experimental/SimpleHMD/Source/SimpleHMD/Private/SimpleHMD.h
+++ b/Engine/Plugins/Experimental/SimpleHMD/Source/SimpleHMD/Private/SimpleHMD.h
@@ -4,9 +4,6 @@
 #include "ISimpleHMDPlugin.h"
 #include "HeadMountedDisplay.h"
 #include "IHeadMountedDisplay.h"
-
-#if SIMPLEHMD_SUPPORTED_PLATFORMS
-
 #include "SceneViewExtension.h"
 
 
@@ -95,9 +92,19 @@ public:
 
 	/** @return	True if the HMD was initialized OK */
 	bool IsInitialized() const;
+
+private:
+
+	FQuat					CurHmdOrientation;
+	FQuat					LastHmdOrientation;
+
+	FRotator				DeltaControlRotation;    // same as DeltaControlOrientation but as rotator
+	FQuat					DeltaControlOrientation; // same as DeltaControlRotation but as quat
+
+	double					LastSensorTime;
+
+	void GetCurrentPose(FQuat& CurrentOrientation);
 };
 
 
 DEFINE_LOG_CATEGORY_STATIC(LogHMD, Log, All);
-
-#endif //SIMPLEHMD_SUPPORTED_PLATFORMS
diff --git a/Engine/Plugins/Experimental/SimpleHMD/Source/SimpleHMD/Public/ISimpleHMDPlugin.h b/Engine/Plugins/Experimental/SimpleHMD/Source/SimpleHMD/Public/ISimpleHMDPlugin.h
--- a/Engine/Plugins/Experimental/SimpleHMD/Source/SimpleHMD/Public/ISimpleHMDPlugin.h
+++ b/Engine/Plugins/Experimental/SimpleHMD/Source/SimpleHMD/Public/ISimpleHMDPlugin.h
@@ -5,8 +5,6 @@
 #include "ModuleManager.h"
 #include "IHeadMountedDisplayModule.h"
 
-#define SIMPLEHMD_SUPPORTED_PLATFORMS PLATFORM_WINDOWS || PLATFORM_MAC || PLATFORM_ANDROID
-
 /**
  * The public interface to this module.  In most cases, this interface is only public to sibling modules 
  * within this plugin.
diff --git a/Engine/Plugins/Experimental/SimpleHMD/Source/SimpleHMD/SimpleHMD.Build.cs b/Engine/Plugins/Experimental/SimpleHMD/Source/SimpleHMD/SimpleHMD.Build.cs
--- a/Engine/Plugins/Experimental/SimpleHMD/Source/SimpleHMD/SimpleHMD.Build.cs
+++ b/Engine/Plugins/Experimental/SimpleHMD/Source/SimpleHMD/SimpleHMD.Build.cs
@@ -20,6 +20,7 @@ namespace UnrealBuildTool.Rules
 					"Core",
 					"CoreUObject",
 					"Engine",
+					"InputCore",
 					"RHI",
 					"RenderCore",
 					"Renderer",
