#ue4 Fix for SimpleHMD build break.  Updated to conform to new API

diff --git a/Engine/Plugins/Experimental/SimpleHMD/Source/SimpleHMD/Private/SimpleHMD.cpp b/Engine/Plugins/Experimental/SimpleHMD/Source/SimpleHMD/Private/SimpleHMD.cpp
--- a/Engine/Plugins/Experimental/SimpleHMD/Source/SimpleHMD/Private/SimpleHMD.cpp
+++ b/Engine/Plugins/Experimental/SimpleHMD/Source/SimpleHMD/Private/SimpleHMD.cpp
@@ -118,7 +118,7 @@ void FSimpleHMD::GetCurrentPose(FQuat& CurrentOrientation)
 	}
 }
 
-void FSimpleHMD::GetCurrentOrientationAndPosition(FQuat& CurrentOrientation, FVector& CurrentPosition, bool bUseOrienationForPlayerCamera, bool bUsePositionForPlayerCamera, const FVector& PositionScale)
+void FSimpleHMD::GetCurrentOrientationAndPosition(FQuat& CurrentOrientation, FVector& CurrentPosition)
 {
 	CurrentPosition = FVector(0.0f, 0.0f, 0.0f);
 
@@ -151,7 +151,7 @@ void FSimpleHMD::ApplyHmdRotation(APlayerController* PC, FRotator& ViewRotation)
 	ViewRotation = FRotator(DeltaControlOrientation * CurHmdOrientation);
 }
 
-void FSimpleHMD::UpdatePlayerCamera(APlayerCameraManager* Camera, struct FMinimalViewInfo& POV)
+void FSimpleHMD::UpdatePlayerCameraRotation(APlayerCameraManager* Camera, struct FMinimalViewInfo& POV)
 {
 	return;
 }
@@ -342,7 +342,7 @@ void FSimpleHMD::SetupView(FSceneViewFamily& InViewFamily, FSceneView& InView)
 	InViewFamily.bUseSeparateRenderTarget = false;
 }
 
-void FSimpleHMD::PreRenderView_RenderThread(FRHICommandListImmediate& RHICmdList, FSceneView& View)
+void FSimpleHMD::PreRenderView_RenderThread(FRHICommandListImmediate& RHICmdList, FSceneView& InView)
 {
 	check(IsInRenderingThread());
 }
diff --git a/Engine/Plugins/Experimental/SimpleHMD/Source/SimpleHMD/Private/SimpleHMD.h b/Engine/Plugins/Experimental/SimpleHMD/Source/SimpleHMD/Private/SimpleHMD.h
--- a/Engine/Plugins/Experimental/SimpleHMD/Source/SimpleHMD/Private/SimpleHMD.h
+++ b/Engine/Plugins/Experimental/SimpleHMD/Source/SimpleHMD/Private/SimpleHMD.h
@@ -29,14 +29,14 @@ public:
 	virtual void SetInterpupillaryDistance(float NewInterpupillaryDistance) override;
 	virtual float GetInterpupillaryDistance() const override;
 
-	virtual void GetCurrentOrientationAndPosition(FQuat& CurrentOrientation, FVector& CurrentPosition, bool bUseOrienationForPlayerCamera, bool bUsePositionForPlayerCamera, const FVector& PositionScale) override;
+	virtual void GetCurrentOrientationAndPosition(FQuat& CurrentOrientation, FVector& CurrentPosition) override;
 	virtual TSharedPtr<class ISceneViewExtension, ESPMode::ThreadSafe> GetViewExtension() override;
 	virtual void ApplyHmdRotation(APlayerController* PC, FRotator& ViewRotation) override;
-	virtual void UpdatePlayerCamera(APlayerCameraManager*, struct FMinimalViewInfo& POV) override;
+	virtual void UpdatePlayerCameraRotation(class APlayerCameraManager* Camera, struct FMinimalViewInfo& POV) override;
 
 	virtual bool IsChromaAbCorrectionEnabled() const override;
 
-	virtual bool Exec( UWorld* InWorld, const TCHAR* Cmd, FOutputDevice& Ar ) override;
+	virtual bool Exec(UWorld* InWorld, const TCHAR* Cmd, FOutputDevice& Ar) override;
 	virtual void OnScreenModeChange(EWindowMode::Type WindowMode) override;
 
 	virtual bool IsPositionalTrackingEnabled() const override;
@@ -65,8 +65,8 @@ public:
 	virtual bool IsStereoEnabled() const override;
 	virtual bool EnableStereo(bool stereo = true) override;
 	virtual void AdjustViewRect(EStereoscopicPass StereoPass, int32& X, int32& Y, uint32& SizeX, uint32& SizeY) const override;
-	virtual void CalculateStereoViewOffset(const EStereoscopicPass StereoPassType, const FRotator& ViewRotation, 
-										   const float MetersToWorld, FVector& ViewLocation) override;
+	virtual void CalculateStereoViewOffset(const EStereoscopicPass StereoPassType, const FRotator& ViewRotation,
+		const float MetersToWorld, FVector& ViewLocation) override;
 	virtual FMatrix GetStereoProjectionMatrix(const EStereoscopicPass StereoPassType, const float FOV) const override;
 	virtual void InitCanvasFromView(FSceneView* InView, UCanvas* Canvas) override;
 	virtual void GetEyeRenderParams_RenderThread(const struct FRenderingCompositePassContext& Context, FVector2D& EyeToSrcUVScaleValue, FVector2D& EyeToSrcUVOffsetValue) const override;
@@ -74,7 +74,8 @@ public:
 	/** ISceneViewExtension interface */
 	virtual void SetupViewFamily(FSceneViewFamily& InViewFamily) override;
 	virtual void SetupView(FSceneViewFamily& InViewFamily, FSceneView& InView) override;
-	virtual void PreRenderView_RenderThread(FSceneView& InView) override;
+	virtual void BeginRenderViewFamily(FSceneViewFamily& InViewFamily) {}
+	virtual void PreRenderView_RenderThread(FRHICommandListImmediate& RHICmdList, FSceneView& InView) override;
 	virtual void PreRenderViewFamily_RenderThread(FRHICommandListImmediate& RHICmdList, FSceneViewFamily& InViewFamily) override;
 
 public:
