Don't draw both eyes in the same distortion draw call
Fix debug builds of the SimpleHMD plugin

diff --git a/Engine/Plugins/Experimental/SimpleHMD/Source/SimpleHMD/Private/SimpleHMD.cpp b/Engine/Plugins/Experimental/SimpleHMD/Source/SimpleHMD/Private/SimpleHMD.cpp
--- a/Engine/Plugins/Experimental/SimpleHMD/Source/SimpleHMD/Private/SimpleHMD.cpp
+++ b/Engine/Plugins/Experimental/SimpleHMD/Source/SimpleHMD/Private/SimpleHMD.cpp
@@ -42,7 +42,7 @@ TSharedPtr< class IXRTrackingSystem, ESPMode::ThreadSafe > FSimpleHMDPlugin::Cre
 
 float FSimpleHMD::GetWorldToMetersScale() const
 {
-	return GWorld ? GWorld->GetWorldSettings()->WorldToMeters : 100.0f;
+	return WorldToMeters;
 }
 
 //---------------------------------------------------
@@ -183,16 +183,20 @@ void FSimpleHMD::DrawDistortionMesh_RenderThread(struct FRenderingCompositePassC
 	FIntPoint ViewportSize = ViewFamily.RenderTarget->GetSizeXY();
 	RHICmdList.SetViewport(0, 0, 0.0f, ViewportSize.X, ViewportSize.Y, 1.0f);
 
-	static const uint32 NumVerts = 8;
-	static const uint32 NumTris = 4;
+	static const uint32 NumVerts = 4;
+	static const uint32 NumTris = 2;
 
-	static const FDistortionVertex Verts[8] =
+	static const FDistortionVertex LeftVerts[] =
 	{
 		// left eye
 		{ FVector2D(-0.9f, -0.9f), FVector2D(0.0f, 1.0f), FVector2D(0.0f, 1.0f), FVector2D(0.0f, 1.0f), 1.0f, 0.0f },
 		{ FVector2D(-0.1f, -0.9f), FVector2D(0.5f, 1.0f), FVector2D(0.5f, 1.0f), FVector2D(0.5f, 1.0f), 1.0f, 0.0f },
 		{ FVector2D(-0.1f, 0.9f), FVector2D(0.5f, 0.0f), FVector2D(0.5f, 0.0f), FVector2D(0.5f, 0.0f), 1.0f, 0.0f },
 		{ FVector2D(-0.9f, 0.9f), FVector2D(0.0f, 0.0f), FVector2D(0.0f, 0.0f), FVector2D(0.0f, 0.0f), 1.0f, 0.0f },
+	};
+
+	static const FDistortionVertex RightVerts[] =
+	{
 		// right eye
 		{ FVector2D(0.1f, -0.9f), FVector2D(0.5f, 1.0f), FVector2D(0.5f, 1.0f), FVector2D(0.5f, 1.0f), 1.0f, 0.0f },
 		{ FVector2D(0.9f, -0.9f), FVector2D(1.0f, 1.0f), FVector2D(1.0f, 1.0f), FVector2D(1.0f, 1.0f), 1.0f, 0.0f },
@@ -201,12 +205,12 @@ void FSimpleHMD::DrawDistortionMesh_RenderThread(struct FRenderingCompositePassC
 	};
 
 	FRHIResourceCreateInfo CreateInfo;
-	FVertexBufferRHIRef VertexBufferRHI = RHICreateVertexBuffer(sizeof(FDistortionVertex) * 8, BUF_Volatile, CreateInfo);
-	void* VoidPtr = RHILockVertexBuffer(VertexBufferRHI, 0, sizeof(FDistortionVertex) * 8, RLM_WriteOnly);
-	FPlatformMemory::Memcpy(VoidPtr, Verts, sizeof(FDistortionVertex) * 8);
+	FVertexBufferRHIRef VertexBufferRHI = RHICreateVertexBuffer(sizeof(FDistortionVertex) * NumVerts, BUF_Volatile, CreateInfo);
+	void* VoidPtr = RHILockVertexBuffer(VertexBufferRHI, 0, sizeof(FDistortionVertex) * NumVerts, RLM_WriteOnly);
+	FPlatformMemory::Memcpy(VoidPtr, View.StereoPass == eSSP_RIGHT_EYE ? RightVerts : LeftVerts, sizeof(FDistortionVertex) * NumVerts);
 	RHIUnlockVertexBuffer(VertexBufferRHI);
 
-	static const uint16 Indices[12] = { /*Left*/ 0, 1, 2, 0, 2, 3, /*Right*/ 4, 5, 6, 4, 6, 7 };
+	static const uint16 Indices[] = { 0, 1, 2, 0, 2, 3 };
 
 	FIndexBufferRHIRef IndexBufferRHI = RHICreateIndexBuffer(sizeof(uint16), sizeof(uint16) * 12, BUF_Volatile, CreateInfo);
 	void* VoidPtr2 = RHILockIndexBuffer(IndexBufferRHI, 0, sizeof(uint16) * 12, RLM_WriteOnly);
@@ -239,7 +243,7 @@ void FSimpleHMD::AdjustViewRect(EStereoscopicPass StereoPass, int32& X, int32& Y
 	}
 }
 
-void FSimpleHMD::CalculateStereoViewOffset(const enum EStereoscopicPass StereoPassType, FRotator& ViewRotation, const float WorldToMeters, FVector& ViewLocation)
+void FSimpleHMD::CalculateStereoViewOffset(const enum EStereoscopicPass StereoPassType, FRotator& ViewRotation, const float InWorldToMeters, FVector& ViewLocation)
 {
 	if( StereoPassType != eSSP_FULL)
 	{
@@ -282,6 +286,11 @@ void FSimpleHMD::SetupViewFamily(FSceneViewFamily& InViewFamily)
 	InViewFamily.EngineShowFlags.MotionBlur = 0;
 	InViewFamily.EngineShowFlags.HMDDistortion = true;
 	InViewFamily.EngineShowFlags.StereoRendering = IsStereoEnabled();
+
+	if (GWorld)
+	{
+		WorldToMeters = GWorld->GetWorldSettings()->WorldToMeters;
+	}
 }
 
 void FSimpleHMD::SetupView(FSceneViewFamily& InViewFamily, FSceneView& InView)
@@ -312,7 +321,8 @@ FSimpleHMD::FSimpleHMD(const FAutoRegister& AutoRegister) :
 	LastHmdOrientation(FQuat::Identity),
 	DeltaControlRotation(FRotator::ZeroRotator),
 	DeltaControlOrientation(FQuat::Identity),
-	LastSensorTime(-1.0)
+	LastSensorTime(-1.0),
+	WorldToMeters(100.0f)
 {
 }
 
diff --git a/Engine/Plugins/Experimental/SimpleHMD/Source/SimpleHMD/Private/SimpleHMD.h b/Engine/Plugins/Experimental/SimpleHMD/Source/SimpleHMD/Private/SimpleHMD.h
--- a/Engine/Plugins/Experimental/SimpleHMD/Source/SimpleHMD/Private/SimpleHMD.h
+++ b/Engine/Plugins/Experimental/SimpleHMD/Source/SimpleHMD/Private/SimpleHMD.h
@@ -68,8 +68,8 @@ public:
 	virtual bool IsStereoEnabled() const override;
 	virtual bool EnableStereo(bool stereo = true) override;
 	virtual void AdjustViewRect(EStereoscopicPass StereoPass, int32& X, int32& Y, uint32& SizeX, uint32& SizeY) const override;
-	virtual void CalculateStereoViewOffset(const EStereoscopicPass StereoPassType, FRotator& ViewRotation,
-		const float MetersToWorld, FVector& ViewLocation) override;
+	virtual void CalculateStereoViewOffset(const enum EStereoscopicPass StereoPassType, FRotator& ViewRotation,
+		const float InWorldToMeters, FVector& ViewLocation) override;
 	virtual FMatrix GetStereoProjectionMatrix(const enum EStereoscopicPass StereoPassType) const override;
 	virtual void GetEyeRenderParams_RenderThread(const struct FRenderingCompositePassContext& Context, FVector2D& EyeToSrcUVScaleValue, FVector2D& EyeToSrcUVOffsetValue) const override;
 
@@ -100,6 +100,7 @@ private:
 	FQuat					DeltaControlOrientation; // same as DeltaControlRotation but as quat
 
 	double					LastSensorTime;
+	float					WorldToMeters;
 
 	void GetHMDOrientation(FQuat& CurrentOrientation);
 };
