Fix compiler warning in SimpleHMD

diff --git a/Engine/Plugins/Experimental/SimpleHMD/Source/SimpleHMD/Private/SimpleHMD.cpp b/Engine/Plugins/Experimental/SimpleHMD/Source/SimpleHMD/Private/SimpleHMD.cpp
--- a/Engine/Plugins/Experimental/SimpleHMD/Source/SimpleHMD/Private/SimpleHMD.cpp
+++ b/Engine/Plugins/Experimental/SimpleHMD/Source/SimpleHMD/Private/SimpleHMD.cpp
@@ -243,7 +243,7 @@ void FSimpleHMD::AdjustViewRect(EStereoscopicPass StereoPass, int32& X, int32& Y
 	}
 }
 
-void FSimpleHMD::CalculateStereoViewOffset(const enum EStereoscopicPass StereoPassType, FRotator& ViewRotation, const float WorldToMeters, FVector& ViewLocation)
+void FSimpleHMD::CalculateStereoViewOffset(const enum EStereoscopicPass StereoPassType, FRotator& ViewRotation, const float InWorldToMeters, FVector& ViewLocation)
 {
 	if( StereoPassType != eSSP_FULL)
 	{
diff --git a/Engine/Plugins/Experimental/SimpleHMD/Source/SimpleHMD/Private/SimpleHMD.h b/Engine/Plugins/Experimental/SimpleHMD/Source/SimpleHMD/Private/SimpleHMD.h
--- a/Engine/Plugins/Experimental/SimpleHMD/Source/SimpleHMD/Private/SimpleHMD.h
+++ b/Engine/Plugins/Experimental/SimpleHMD/Source/SimpleHMD/Private/SimpleHMD.h
@@ -69,7 +69,7 @@ public:
 	virtual bool EnableStereo(bool stereo = true) override;
 	virtual void AdjustViewRect(EStereoscopicPass StereoPass, int32& X, int32& Y, uint32& SizeX, uint32& SizeY) const override;
 	virtual void CalculateStereoViewOffset(const enum EStereoscopicPass StereoPassType, FRotator& ViewRotation,
-		const float WorldToMeters, FVector& ViewLocation) override;
+		const float InWorldToMeters, FVector& ViewLocation) override;
 	virtual FMatrix GetStereoProjectionMatrix(const enum EStereoscopicPass StereoPassType) const override;
 	virtual void GetEyeRenderParams_RenderThread(const struct FRenderingCompositePassContext& Context, FVector2D& EyeToSrcUVScaleValue, FVector2D& EyeToSrcUVOffsetValue) const override;
 
