Fixing PVS 7.7 Issues:

warning V512: A call of the 'memcpy' function will lead to the 'Indices' buffer becoming out of range.

warning V512: A call of the 'memcpy' function will lead to underflow of the buffer 'DiscoveryLayer.Universes'.

diff --git a/Engine/Plugins/Experimental/SimpleHMD/Source/SimpleHMD/Private/SimpleHMD.cpp b/Engine/Plugins/Experimental/SimpleHMD/Source/SimpleHMD/Private/SimpleHMD.cpp
--- a/Engine/Plugins/Experimental/SimpleHMD/Source/SimpleHMD/Private/SimpleHMD.cpp
+++ b/Engine/Plugins/Experimental/SimpleHMD/Source/SimpleHMD/Private/SimpleHMD.cpp
@@ -212,9 +212,9 @@ void FSimpleHMD::DrawDistortionMesh_RenderThread(struct FHeadMountedDisplayPassC
 
 	static const uint16 Indices[] = { 0, 1, 2, 0, 2, 3 };
 
-	FIndexBufferRHIRef IndexBufferRHI = RHICreateIndexBuffer(sizeof(uint16), sizeof(uint16) * 12, BUF_Volatile, CreateInfo);
-	void* VoidPtr2 = RHILockBuffer(IndexBufferRHI, 0, sizeof(uint16) * 12, RLM_WriteOnly);
-	FPlatformMemory::Memcpy(VoidPtr2, Indices, sizeof(uint16) * 12);
+	FIndexBufferRHIRef IndexBufferRHI = RHICreateIndexBuffer(sizeof(uint16), sizeof(uint16) * 6, BUF_Volatile, CreateInfo);
+	void* VoidPtr2 = RHILockBuffer(IndexBufferRHI, 0, sizeof(uint16) * 6, RLM_WriteOnly);
+	FPlatformMemory::Memcpy(VoidPtr2, Indices, sizeof(uint16) * 6);
 	RHIUnlockBuffer(IndexBufferRHI);
 
 	RHICmdList.SetStreamSource(0, VertexBufferRHI, 0);
diff --git a/Engine/Plugins/VirtualProduction/DMX/DMXProtocol/Source/DMXProtocolSACN/Private/DMXProtocolSACN.cpp b/Engine/Plugins/VirtualProduction/DMX/DMXProtocol/Source/DMXProtocolSACN/Private/DMXProtocolSACN.cpp
--- a/Engine/Plugins/VirtualProduction/DMX/DMXProtocol/Source/DMXProtocolSACN/Private/DMXProtocolSACN.cpp
+++ b/Engine/Plugins/VirtualProduction/DMX/DMXProtocol/Source/DMXProtocolSACN/Private/DMXProtocolSACN.cpp
@@ -278,7 +278,7 @@ bool FDMXProtocolSACN::SendDiscovery(const TArray<uint16>& Universes)
 	FDMXProtocolUDPE131DiscoveryLayerPacket DiscoveryLayer;
 	if (Universes.Num() == ACN_DMX_SIZE)
 	{
-		FMemory::Memcpy(DiscoveryLayer.Universes, Universes.GetData(), ACN_DMX_SIZE);
+		FMemory::Memcpy(DiscoveryLayer.Universes, Universes.GetData(), ACN_DMX_SIZE * sizeof(uint16));
 	}
 	else
 	{
