Set supported platforms for SimpleHMD plugin.
Clang static analysis 63879a52d76995cb14b121ff

diff --git a/Engine/Plugins/Experimental/SimpleHMD/SimpleHMD.uplugin b/Engine/Plugins/Experimental/SimpleHMD/SimpleHMD.uplugin
--- a/Engine/Plugins/Experimental/SimpleHMD/SimpleHMD.uplugin
+++ b/Engine/Plugins/Experimental/SimpleHMD/SimpleHMD.uplugin
@@ -14,6 +14,16 @@
 	"CanContainContent" : false,
 	"IsBetaVersion" : false,
 	"Installed" : false,
+	"SupportedTargetPlatforms": [
+		"Win64",
+		"Linux",
+		"Android",
+		"HoloLens",
+		"PS4",
+		"PS5",
+		"Mac",
+		"IOS"
+	],
 	"Modules" :
 	[
 		{
