Fixed up platforms for SimpleHMD plugin.

#rnx

diff --git a/Engine/Plugins/Experimental/SimpleHMD/SimpleHMD.uplugin b/Engine/Plugins/Experimental/SimpleHMD/SimpleHMD.uplugin
--- a/Engine/Plugins/Experimental/SimpleHMD/SimpleHMD.uplugin
+++ b/Engine/Plugins/Experimental/SimpleHMD/SimpleHMD.uplugin
@@ -19,8 +19,6 @@
 		"Linux",
 		"Android",
 		"HoloLens",
-		"PS4",
-		"PS5",
 		"Mac",
 		"IOS"
 	],
